Adding a new Plugin Template Tool to assist in managing content belonging to a PluginTemplate. The tool can temporarily mount a template content folder allowing assets to be moved or edited safely within the editor.

diff --git a/Engine/Plugins/Editor/PluginBrowser/Source/PluginBrowser/Private/PluginBrowserModule.h b/Engine/Plugins/Editor/PluginBrowser/Source/PluginBrowser/Private/PluginBrowserModule.h
--- a/Engine/Plugins/Editor/PluginBrowser/Source/PluginBrowser/Private/PluginBrowserModule.h
+++ b/Engine/Plugins/Editor/PluginBrowser/Source/PluginBrowser/Private/PluginBrowserModule.h
@@ -76,7 +76,7 @@ public:
 	/** Spawns the plugin creator tab with a specific wizard definition */
 	virtual TSharedRef<SDockTab> SpawnPluginCreatorTab(const FSpawnTabArgs& SpawnTabArgs, TSharedPtr<IPluginWizardDefinition> PluginWizardDefinition) override;
 
-	const TArray<TSharedRef<FPluginTemplateDescription>>& GetAddedPluginTemplates() const { return AddedPluginTemplates; }
+	virtual const TArray<TSharedRef<FPluginTemplateDescription>>& GetAddedPluginTemplates() const override { return AddedPluginTemplates; }
 
 	const TArray<TPair<FOnPluginBeingEdited, FPluginEditorExtensionHandle>>& GetCustomizePluginEditingDelegates() { return CustomizePluginEditingDelegates; }
 
diff --git a/Engine/Plugins/Experimental/PluginTemplateTool/PluginTemplateTool.uplugin b/Engine/Plugins/Experimental/PluginTemplateTool/PluginTemplateTool.uplugin
new file mode 100644
--- /dev/null
+++ b/Engine/Plugins/Experimental/PluginTemplateTool/PluginTemplateTool.uplugin
@@ -0,0 +1,32 @@
+{
+	"FileVersion" : 3,
+	"Version" : 1,
+	"VersionName" : "1.0",
+	"FriendlyName" : "Plugin Template Tool",
+	"Description" : "Editor plugin for managing plugin template content folders.",
+	"Category" : "Editor",
+	"CreatedBy" : "Epic Games, Inc.",
+	"CreatedByURL" : "http://epicgames.com",
+	"DocsURL" : "",
+	"MarketplaceURL" : "",
+	"SupportURL" : "",
+	"EnabledByDefault" : false,
+	"CanContainContent" : false,
+	"IsBetaVersion" : true,
+	"Installed" : false,
+	
+	"Modules" :
+	[
+		{
+			"Name": "PluginTemplateTool",
+			"Type": "Editor",
+			"LoadingPhase": "PreDefault"
+		}
+	],
+	"Plugins": [
+		{
+			"Name": "PluginBrowser",
+			"Enabled": true
+		},
+	]
+}
\ No newline at end of file
diff --git a/Engine/Plugins/Experimental/PluginTemplateTool/Source/PluginTemplateTool.Build.cs b/Engine/Plugins/Experimental/PluginTemplateTool/Source/PluginTemplateTool.Build.cs
new file mode 100644
--- /dev/null
+++ b/Engine/Plugins/Experimental/PluginTemplateTool/Source/PluginTemplateTool.Build.cs
@@ -0,0 +1,35 @@
+// Copyright Epic Games, Inc. All Rights Reserved.
+
+using UnrealBuildTool;
+
+public class PluginTemplateTool : ModuleRules
+{
+	public PluginTemplateTool(ReadOnlyTargetRules Target) : base(Target)
+	{
+		PrivateIncludePathModuleNames.AddRange(
+			new string[] {
+			}
+		);
+
+		PublicDependencyModuleNames.AddRange(
+			new string[] { 
+				"Core",
+				"CoreUObject",
+				"InputCore",
+				"Engine",
+				"Slate",
+				"SlateCore",
+			}
+		);
+		
+		PrivateDependencyModuleNames.AddRange(
+			new string[] {
+				"EditorFramework",
+				"PluginBrowser",
+				"ToolWidgets",
+				"EditorWidgets",
+				"WorkspaceMenuStructure",
+			}
+		);
+	}
+}
diff --git a/Engine/Plugins/Experimental/PluginTemplateTool/Source/PluginTemplateTool/Private/PluginTemplateToolModule.cpp b/Engine/Plugins/Experimental/PluginTemplateTool/Source/PluginTemplateTool/Private/PluginTemplateToolModule.cpp
new file mode 100644
--- /dev/null
+++ b/Engine/Plugins/Experimental/PluginTemplateTool/Source/PluginTemplateTool/Private/PluginTemplateToolModule.cpp
@@ -0,0 +1,55 @@
+// Copyright Epic Games, Inc. All Rights Reserved.
+
+#include "Modules/ModuleInterface.h"
+#include "SPluginTemplateBrowser.h"
+#include "Widgets/Docking/SDockTab.h"
+#include "WorkspaceMenuStructure.h"
+#include "WorkspaceMenuStructureModule.h"
+
+#define LOCTEXT_NAMESPACE "PluginTemplateTool"
+
+class FPluginTemplateToolModule : public IModuleInterface
+{
+public:
+	// IModuleInterface interface
+	virtual void StartupModule() override;
+	virtual void ShutdownModule() override;
+	// End of IModuleInterface interface
+
+private:
+	static const FName PluginTemplateToolTabName;
+	TSharedRef<SDockTab> HandleSpawnPluginTemplatesTab(const FSpawnTabArgs& Args);
+};
+
+const FName FPluginTemplateToolModule::PluginTemplateToolTabName = TEXT("PluginTemplateTool");
+
+///////////////////////////////////////////
+
+IMPLEMENT_MODULE(FPluginTemplateToolModule, PluginTemplateTool);
+
+
+void FPluginTemplateToolModule::StartupModule()
+{
+	FGlobalTabmanager::Get()->RegisterNomadTabSpawner(PluginTemplateToolTabName, FOnSpawnTab::CreateRaw(this, &FPluginTemplateToolModule::HandleSpawnPluginTemplatesTab))
+		.SetDisplayName(LOCTEXT("PluginTemplateTitle", "Plugin Template Tool"))
+		.SetTooltipText(LOCTEXT("PluginTemplateTooltip", "Open Plugin Template Tool window."))
+		.SetGroup(WorkspaceMenu::GetMenuStructure().GetToolsCategory())
+		.SetIcon(FSlateIcon(FAppStyle::GetAppStyleSetName(), "LevelEditor.Audit"));
+	FGlobalTabmanager::Get()->RegisterDefaultTabWindowSize(PluginTemplateToolTabName, FVector2D(1080, 600));
+}
+
+void FPluginTemplateToolModule::ShutdownModule()
+{
+	FGlobalTabmanager::Get()->UnregisterNomadTabSpawner(PluginTemplateToolTabName);
+}
+
+TSharedRef<SDockTab> FPluginTemplateToolModule::HandleSpawnPluginTemplatesTab(const FSpawnTabArgs& Args)
+{
+	return SNew(SDockTab)
+		.TabRole(ETabRole::NomadTab)
+		[
+			SNew(SPluginTemplateBrowser)
+		];
+}
+
+#undef LOCTEXT_NAMESPACE
diff --git a/Engine/Plugins/Experimental/PluginTemplateTool/Source/PluginTemplateTool/Private/SPluginTemplateBrowser.cpp b/Engine/Plugins/Experimental/PluginTemplateTool/Source/PluginTemplateTool/Private/SPluginTemplateBrowser.cpp
new file mode 100644
--- /dev/null
+++ b/Engine/Plugins/Experimental/PluginTemplateTool/Source/PluginTemplateTool/Private/SPluginTemplateBrowser.cpp
@@ -0,0 +1,189 @@
+// Copyright Epic Games, Inc. All Rights Reserved.
+
+#include "SPluginTemplateBrowser.h"
+
+#include "Misc/MessageDialog.h"
+#include "Misc/PackageName.h"
+#include "IPluginBrowser.h"
+#include "Widgets/Input/SButton.h"
+#include "Widgets/Views/SListView.h"
+#include "Widgets/Views/STableViewBase.h"
+
+#define LOCTEXT_NAMESPACE "PluginTemplateBrowser"
+
+void SPluginTemplateBrowser::Construct(const FArguments& InArgs)
+{
+	const TArray<TSharedRef<FPluginTemplateDescription>>& PluginTemplateDescriptions = IPluginBrowser::Get().GetAddedPluginTemplates();
+	for (const auto& Item : PluginTemplateDescriptions)
+	{
+		TSharedRef<FPluginTemplateListItem> NewItem(new FPluginTemplateListItem(Item->Name, Item->OnDiskPath));
+		PluginTemplateListItems.Add(NewItem);
+	}
+
+	ChildSlot
+		.Padding(FMargin(8))
+	[
+		SNew(SVerticalBox)
+			+ SVerticalBox::Slot()
+			.FillHeight(1.0f)
+			[
+				SNew(SPluginTemplateListView)
+					.SelectionMode(ESelectionMode::Single)
+					.ItemHeight(24)
+					.ListItemsSource(&PluginTemplateListItems)
+					.OnGenerateRow(this, &SPluginTemplateBrowser::OnGenerateWidgetForTemplateListView)
+					.HeaderRow
+					(
+						SNew(SHeaderRow)
+						+ SHeaderRow::Column("TemplateNames")
+						.DefaultLabel(NSLOCTEXT("PluginTemplateBrowser", "TemplateName", "Template"))
+						.FillWidth(0.7f)
+						+ SHeaderRow::Column("TemplateActions")
+						.DefaultLabel(NSLOCTEXT("PluginTemplateBrowser", "TemplateActions", "Actions"))
+						.FillWidth(0.36f)
+						+ SHeaderRow::Column("TemplatePaths")
+						.DefaultLabel(NSLOCTEXT("PluginTemplateBrowser", "TemplatePath", "Path"))
+					)
+			]
+	];
+}
+
+TSharedRef<ITableRow> SPluginTemplateBrowser::OnGenerateWidgetForTemplateListView(TSharedPtr<FPluginTemplateListItem> InItem, const TSharedRef<STableViewBase>& OwnerTable)
+{
+	class STemplateItemWidget : public SMultiColumnTableRow<TSharedPtr<FPluginTemplateListItem>>
+	{
+	public:
+		SLATE_BEGIN_ARGS(STemplateItemWidget) {}
+		SLATE_END_ARGS()
+
+		void Construct(const FArguments& InArgs, const TSharedRef<STableViewBase>& InOwnerTable, TSharedPtr<FPluginTemplateListItem> InListItem)
+		{
+			Item = InListItem;
+
+			SMultiColumnTableRow<TSharedPtr<FPluginTemplateListItem>>::Construct(FSuperRowType::FArguments(), InOwnerTable);
+		}
+
+		TSharedRef<SWidget> GenerateWidgetForColumn(const FName& ColumnName)
+		{
+			if (ColumnName == "TemplateNames")
+			{
+				return
+					SNew(STextBlock)
+					.Text(Item->TemplateName);
+			}
+			else if (ColumnName == "TemplateActions")
+			{
+				return
+					SNew(SHorizontalBox)
+
+					// Mount button
+					+ SHorizontalBox::Slot()
+					.AutoWidth()
+					.Padding(2.0f, 0.0f)
+					[
+						SNew(SButton)
+						.Visibility(Item.ToSharedRef(), &FPluginTemplateListItem::GetVisibilityBasedOnUnmountedState)
+						.Text(LOCTEXT("Mount", "Mount"))
+						.OnClicked(Item.ToSharedRef(), &FPluginTemplateListItem::OnMountClicked)
+					]
+
+					// Unmount button
+					+ SHorizontalBox::Slot()
+					.AutoWidth()
+					.Padding(2.0f, 0.0f)
+					[
+						SNew(SButton)
+						.Visibility(Item.ToSharedRef(), &FPluginTemplateListItem::GetVisibilityBasedOnMountedState)
+						.Text(LOCTEXT("Unmount", "Unmount"))
+						.OnClicked(Item.ToSharedRef(), &FPluginTemplateListItem::OnUnmountClicked)
+					];
+			}
+			else if (ColumnName == "TemplatePaths")
+			{
+				return
+					SNew(STextBlock)
+					.Text(FText::FromString(Item->OnDiskPath));
+			}
+			else
+			{
+				return SNew(STextBlock).Text(LOCTEXT("UnknownColumn", "Unknown Column"));
+			}
+
+		}
+
+		TSharedPtr<FPluginTemplateListItem> Item;
+	};
+
+	return SNew(STemplateItemWidget, OwnerTable, InItem);
+}
+
+SPluginTemplateBrowser::FPluginTemplateListItem::FPluginTemplateListItem(FText InTemplateName, FString InOnDiskPath)
+	: TemplateName(InTemplateName)
+	, OnDiskPath(InOnDiskPath)
+{
+	bIsMounted = FPackageName::MountPointExists(GetRootPath());
+}
+
+FString SPluginTemplateBrowser::FPluginTemplateListItem::GetRootPath() const
+{
+	return TEXT("/") + FPaths::GetBaseFilename(OnDiskPath) + TEXT("/");
+}
+
+FString SPluginTemplateBrowser::FPluginTemplateListItem::GetContentPath() const
+{
+	return OnDiskPath / TEXT("Content/");
+}
+
+FReply SPluginTemplateBrowser::FPluginTemplateListItem::OnMountClicked()
+{
+	const FString& RootPath = GetRootPath();
+	const FString& ContentPath = GetContentPath();
+
+	FText FailReason;
+	if (ContentPath[0] == TEXT('/'))
+	{
+		FailReason = LOCTEXT("MountPluginTemplate_InvalidContentPath", "Invalid ContentPath, should not start with '/'! Example: '../../../ProjectName/Content/'");
+	}
+
+	if (RootPath[0] != TEXT('/'))
+	{
+		FailReason = LOCTEXT("MountPluginTemplate_InvalidRootPath", "Invalid RootPath, should start with a '/'! Example: '/Game/'");
+	}
+
+	if (!FPaths::DirectoryExists(ContentPath))
+	{
+		FailReason = FText::Format(LOCTEXT("MountPluginTemplate_ContentNotExists", "Content Path: {0} does not exist"), FText::FromString(FPaths::ConvertRelativePathToFull(ContentPath)));
+	}
+
+	if (FailReason.IsEmpty())
+	{
+		FPackageName::RegisterMountPoint(RootPath, ContentPath);
+		bIsMounted = true;
+	}
+	else
+	{
+		FMessageDialog::Open(EAppMsgType::Ok, FailReason, LOCTEXT("UnableToMountPluginTemplate", "Unable to mount plugin template"));
+	}
+
+	return FReply::Handled();
+}
+
+FReply SPluginTemplateBrowser::FPluginTemplateListItem::OnUnmountClicked()
+{
+	check(FPackageName::MountPointExists(GetRootPath()));
+
+	FPackageName::UnRegisterMountPoint(GetRootPath(), GetContentPath());
+	bIsMounted = false;
+
+	return FReply::Handled();
+}
+
+EVisibility SPluginTemplateBrowser::FPluginTemplateListItem::GetVisibilityBasedOnMountedState() const
+{
+	return bIsMounted ? EVisibility::Visible : EVisibility::Hidden;
+}
+
+EVisibility SPluginTemplateBrowser::FPluginTemplateListItem::GetVisibilityBasedOnUnmountedState() const
+{
+	return !bIsMounted ? EVisibility::Visible : EVisibility::Hidden;
+}
\ No newline at end of file
diff --git a/Engine/Plugins/Experimental/PluginTemplateTool/Source/PluginTemplateTool/Private/SPluginTemplateBrowser.h b/Engine/Plugins/Experimental/PluginTemplateTool/Source/PluginTemplateTool/Private/SPluginTemplateBrowser.h
new file mode 100644
--- /dev/null
+++ b/Engine/Plugins/Experimental/PluginTemplateTool/Source/PluginTemplateTool/Private/SPluginTemplateBrowser.h
@@ -0,0 +1,57 @@
+// Copyright Epic Games, Inc. All Rights Reserved.
+
+#pragma once
+
+#include "Widgets/SCompoundWidget.h"
+
+class ITableRow;
+class STableViewBase;
+
+template <typename ItemType> class SListView;
+
+//////////////////////////////////////////////////////////////////////////
+// SPluginTemplateBrowser
+
+class SPluginTemplateBrowser : public SCompoundWidget
+{
+public:
+	SLATE_BEGIN_ARGS(SPluginTemplateBrowser) {}
+	SLATE_END_ARGS()
+
+public:
+	void Construct(const FArguments& InArgs);
+
+private:
+	/** An item in the plugin template list */
+	struct FPluginTemplateListItem
+	{
+		FPluginTemplateListItem(FText InTemplateName, FString InOnDiskPath);
+
+		FText TemplateName;
+		FString OnDiskPath;
+		bool bIsMounted =false;
+
+		FString GetRootPath() const;
+		FString GetContentPath() const;
+
+		FReply OnMountClicked();
+		FReply OnUnmountClicked();
+
+		/** Returns Visible for templates that are currently mounted, otherwise Hidden */
+		EVisibility GetVisibilityBasedOnMountedState() const;
+
+		/** Returns Hidden for a template that are currently mounted, otherwise Visible */
+		EVisibility GetVisibilityBasedOnUnmountedState() const;
+	};
+
+	TSharedRef<ITableRow> OnGenerateWidgetForTemplateListView(TSharedPtr<FPluginTemplateListItem> InItem, const TSharedRef<STableViewBase>& OwnerTable);
+
+private:
+	typedef SListView<TSharedPtr<FPluginTemplateListItem>> SPluginTemplateListView;
+
+	/** List items for the plugin template list */
+	TArray<TSharedPtr<FPluginTemplateListItem>> PluginTemplateListItems;
+
+	/** List of all known plugin templates */
+	TSharedPtr<SPluginTemplateListView> PluginTemplateListView;
+};
\ No newline at end of file
diff --git a/Engine/Source/Editor/UnrealEd/Public/Features/IPluginsEditorFeature.h b/Engine/Source/Editor/UnrealEd/Public/Features/IPluginsEditorFeature.h
--- a/Engine/Source/Editor/UnrealEd/Public/Features/IPluginsEditorFeature.h
+++ b/Engine/Source/Editor/UnrealEd/Public/Features/IPluginsEditorFeature.h
@@ -125,6 +125,11 @@ public:
 	  */
 	virtual void UnregisterPluginEditorExtension(FPluginEditorExtensionHandle ExtensionHandle) = 0;
 
+	/**
+	  * Returns the currently registered plugin templates.
+	  */
+	virtual const TArray<TSharedRef<FPluginTemplateDescription>>& GetAddedPluginTemplates() const = 0;
+
 	/**
 	 * Open the plugin editor for a plugin
 	 */
