Fixing localization build warning

diff --git a/Engine/Plugins/Experimental/PluginTemplateTool/Source/PluginTemplateTool/Private/SPluginTemplateBrowser.cpp b/Engine/Plugins/Experimental/PluginTemplateTool/Source/PluginTemplateTool/Private/SPluginTemplateBrowser.cpp
--- a/Engine/Plugins/Experimental/PluginTemplateTool/Source/PluginTemplateTool/Private/SPluginTemplateBrowser.cpp
+++ b/Engine/Plugins/Experimental/PluginTemplateTool/Source/PluginTemplateTool/Private/SPluginTemplateBrowser.cpp
@@ -186,4 +186,6 @@ EVisibility SPluginTemplateBrowser::FPluginTemplateListItem::GetVisibilityBasedO
 EVisibility SPluginTemplateBrowser::FPluginTemplateListItem::GetVisibilityBasedOnUnmountedState() const
 {
 	return !bIsMounted ? EVisibility::Visible : EVisibility::Hidden;
-}
\ No newline at end of file
+}
+
+#undef LOCTEXT_NAMESPACE
\ No newline at end of file
